DirectoryPlaceholder: Adds the option to automatically generate placeholder assets when new folders are created, and delete them when otherwise empty folders are deleted.

diff --git a/Engine/Plugins/Developer/DirectoryPlaceholder/Source/DirectoryPlaceholder/DirectoryPlaceholder.build.cs b/Engine/Plugins/Developer/DirectoryPlaceholder/Source/DirectoryPlaceholder/DirectoryPlaceholder.build.cs
--- a/Engine/Plugins/Developer/DirectoryPlaceholder/Source/DirectoryPlaceholder/DirectoryPlaceholder.build.cs
+++ b/Engine/Plugins/Developer/DirectoryPlaceholder/Source/DirectoryPlaceholder/DirectoryPlaceholder.build.cs
@@ -20,6 +20,7 @@ namespace UnrealBuildTool.Rules
 					"ContentBrowserData",
 					"Core",
 					"CoreUObject",
+					"DeveloperSettings",
 					"UnrealEd",
 				}
 			);
diff --git a/Engine/Plugins/Developer/DirectoryPlaceholder/Source/DirectoryPlaceholder/Private/DirectoryPlaceholderModule.cpp b/Engine/Plugins/Developer/DirectoryPlaceholder/Source/DirectoryPlaceholder/Private/DirectoryPlaceholderModule.cpp
--- a/Engine/Plugins/Developer/DirectoryPlaceholder/Source/DirectoryPlaceholder/Private/DirectoryPlaceholderModule.cpp
+++ b/Engine/Plugins/Developer/DirectoryPlaceholder/Source/DirectoryPlaceholder/Private/DirectoryPlaceholderModule.cpp
@@ -2,4 +2,163 @@
 
 #include "Modules/ModuleManager.h"
 
-IMPLEMENT_MODULE(FDefaultModuleImpl, DirectoryPlaceholder);
+#include "AssetRegistry/AssetRegistryModule.h"
+#include "AssetToolsModule.h"
+#include "ContentBrowserItemPath.h"
+#include "ContentBrowserModule.h"
+#include "DirectoryPlaceholder.h"
+#include "DirectoryPlaceholderFactory.h"
+#include "DirectoryPlaceholderSettings.h"
+#include "DirectoryWatcherModule.h"
+#include "FileHelpers.h"
+#include "IDirectoryWatcher.h"
+#include "ObjectTools.h"
+
+DEFINE_LOG_CATEGORY_STATIC(LogDirectoryPlaceholder, Log, All)
+
+#define LOCTEXT_NAMESPACE "DirectoryPlaceholderModule"
+
+/** 
+ * Directory Placeholder Module
+ * Manages automatic creation and deletion of placeholder assets, allowing folders to be checked in to source control.
+ */
+class FDirectoryPlaceholderModule : public IModuleInterface
+{
+public:
+
+	//~ Begin IModuleInterface interface
+	virtual void StartupModule() override;
+	virtual void ShutdownModule() override;
+	//~ End IModuleInterface interface
+
+private:
+	/** Callback when the directory watcher detects a file/folder change */
+	void OnDirectoryChanged(const TArray<FFileChangeData>& InFileChanges);
+
+	/** Executes when one or more folders are being deleted in the content browser */
+	void OnDeleteFolders(const TArray<FContentBrowserItemPath>& PathsToDelete);
+
+private:
+	FDelegateHandle EngineDirectoryWatcherHandle;
+	FDelegateHandle OnDeleteFoldersHandle;
+};
+
+void FDirectoryPlaceholderModule::StartupModule()
+{
+	// Register a callback with the directory watcher to be notified about file/folder changes
+	FDirectoryWatcherModule& DirectoryWatcherModule = FModuleManager::GetModuleChecked<FDirectoryWatcherModule>("DirectoryWatcher");
+	if (IDirectoryWatcher* DirectoryWatcher = DirectoryWatcherModule.Get())
+	{
+		DirectoryWatcher->RegisterDirectoryChangedCallback_Handle(
+			FPaths::ProjectContentDir(),
+			IDirectoryWatcher::FDirectoryChanged::CreateRaw(this, &FDirectoryPlaceholderModule::OnDirectoryChanged),
+			EngineDirectoryWatcherHandle,
+			IDirectoryWatcher::WatchOptions::IncludeDirectoryChanges
+		);
+	}
+
+	// Register a callback to execute when one or more folders are being deleted in the content browser
+	FContentBrowserModule& ContentBrowserModule = FModuleManager::Get().GetModuleChecked<FContentBrowserModule>(TEXT("ContentBrowser"));
+	OnDeleteFoldersHandle = ContentBrowserModule.GetOnDeleteFolders().AddRaw(this, &FDirectoryPlaceholderModule::OnDeleteFolders);
+}
+
+void FDirectoryPlaceholderModule::ShutdownModule()
+{
+	FDirectoryWatcherModule& DirectoryWatcherModule = FModuleManager::GetModuleChecked<FDirectoryWatcherModule>("DirectoryWatcher");
+	if (IDirectoryWatcher* DirectoryWatcher = DirectoryWatcherModule.Get())
+	{
+		DirectoryWatcher->UnregisterDirectoryChangedCallback_Handle(FPaths::ProjectContentDir(), EngineDirectoryWatcherHandle);
+	}
+
+	FContentBrowserModule& ContentBrowserModule = FModuleManager::Get().GetModuleChecked<FContentBrowserModule>(TEXT("ContentBrowser"));
+	ContentBrowserModule.GetOnDeleteFolders().Remove(OnDeleteFoldersHandle);
+}
+
+void FDirectoryPlaceholderModule::OnDirectoryChanged(const TArray<FFileChangeData>& InFileChanges)
+{
+	const UDirectoryPlaceholderSettings* const Settings = GetDefault<UDirectoryPlaceholderSettings>();
+	if (!Settings->bAutomaticallyCreatePlaceholders)
+	{
+		return;
+	}
+
+	for (const FFileChangeData& FileChange : InFileChanges)
+	{
+		// Check if the FileChangeData represent a directory that was just added
+		if (FPaths::DirectoryExists(FileChange.Filename) && (FileChange.Action == FFileChangeData::FCA_Added))
+		{
+			// Create a new UDirectoryPlaceholder asset in the directory that was just added
+			const FString AssetName = FString(TEXT("UE_Placeholder"));
+			const FString PackagePath = FPackageName::FilenameToLongPackageName(FileChange.Filename);
+
+			IAssetTools& AssetTools = FModuleManager::GetModuleChecked<FAssetToolsModule>("AssetTools").Get();
+			const UObject* const PlaceholderAsset = AssetTools.CreateAsset(AssetName, PackagePath, UDirectoryPlaceholder::StaticClass(), NewObject<UDirectoryPlaceholderFactory>());
+
+			// If the asset was created successfully, save it now, because it will be hidden from the user by default
+			// If source control is enabled, it will also be marked for add.
+			if (PlaceholderAsset)
+			{
+				TArray<UPackage*> PackagesToSave = { PlaceholderAsset->GetPackage() };
+				const bool bOnlyDirty = false;
+				if (UEditorLoadingAndSavingUtils::SavePackages(PackagesToSave, bOnlyDirty))
+				{
+					UE_LOG(LogDirectoryPlaceholder, Log, TEXT("New Directory Placeholder was created and saved in %s"), *PackagePath);
+				}
+				else
+				{
+					UE_LOG(LogDirectoryPlaceholder, Log, TEXT("New Directory Placeholder was created in %s, but could not be saved."), *PackagePath);
+				}
+			}
+			else
+			{
+				UE_LOG(LogDirectoryPlaceholder, Warning, TEXT("Failed to create new Directory Plcaeholder in %s"), *PackagePath);
+			}
+		}
+	}
+}
+
+void FDirectoryPlaceholderModule::OnDeleteFolders(const TArray<FContentBrowserItemPath>& PathsToDelete)
+{
+	const UDirectoryPlaceholderSettings* const Settings = GetDefault<UDirectoryPlaceholderSettings>();
+	if (!Settings->bAutomaticallyCreatePlaceholders)
+	{
+		return;
+	}
+
+	FAssetRegistryModule& AssetRegistryModule = FModuleManager::Get().LoadModuleChecked<FAssetRegistryModule>(TEXT("AssetRegistry"));
+
+	// Set up a recursive asset registry filter for the paths being deleted
+	FARFilter Filter;
+	Filter.bRecursivePaths = true;
+	for (const FContentBrowserItemPath& Path : PathsToDelete)
+	{
+		Filter.PackagePaths.Add(Path.GetInternalPathName());
+	}
+
+	// Find all of the assets in the folders being deleted
+	TArray<FAssetData> AssetDataList;
+	AssetRegistryModule.Get().GetAssets(Filter, AssetDataList);
+
+	// Check if all of the assets being deleted are placeholders, or if there are any others present
+	bool bAllPlaceholders = true;
+	for (const FAssetData& AssetData : AssetDataList)
+	{
+		if (AssetData.GetClass() != UDirectoryPlaceholder::StaticClass())
+		{
+			bAllPlaceholders = false;
+		}
+	}
+
+	// If all of the assets are directory placeholders, delete them from these folders
+	if (bAllPlaceholders)
+	{
+		const bool bShowConfirmation = false;
+		const int32 NumAssetsDeleted = ObjectTools::DeleteAssets(AssetDataList, bShowConfirmation);
+
+		UE_LOG(LogDirectoryPlaceholder, Verbose, TEXT("Deleted %d Directory Placeholders"), NumAssetsDeleted);
+	}
+}
+
+IMPLEMENT_MODULE(FDirectoryPlaceholderModule, DirectoryPlaceholder);
+
+#undef LOCTEXT_NAMESPACE
diff --git a/Engine/Plugins/Developer/DirectoryPlaceholder/Source/DirectoryPlaceholder/Private/DirectoryPlaceholderSettings.h b/Engine/Plugins/Developer/DirectoryPlaceholder/Source/DirectoryPlaceholder/Private/DirectoryPlaceholderSettings.h
new file mode 100644
--- /dev/null
+++ b/Engine/Plugins/Developer/DirectoryPlaceholder/Source/DirectoryPlaceholder/Private/DirectoryPlaceholderSettings.h
@@ -0,0 +1,23 @@
+// Copyright Epic Games, Inc. All Rights Reserved.
+#pragma  once
+
+#include "Engine/DeveloperSettings.h"
+
+#include "DirectoryPlaceholderSettings.generated.h"
+
+
+/**
+ * Directory Placeholder Settings
+ */
+UCLASS(config = EditorPerProjectUserSettings)
+class UDirectoryPlaceholderSettings : public UDeveloperSettings
+{
+	GENERATED_BODY()
+
+public:
+	UDirectoryPlaceholderSettings() = default;
+
+	/** If enabled, Directory Placeholder assets will be automatically created in new folders, and automatically deleted from folders being deleted */
+	UPROPERTY(EditAnywhere, Category = "Default")
+	bool bAutomaticallyCreatePlaceholders = true;
+};
diff --git a/Engine/Plugins/Editor/ContentBrowser/ContentBrowserAssetDataSource/Source/ContentBrowserAssetDataSource/Private/ContentBrowserAssetDataCore.cpp b/Engine/Plugins/Editor/ContentBrowser/ContentBrowserAssetDataSource/Source/ContentBrowserAssetDataSource/Private/ContentBrowserAssetDataCore.cpp
--- a/Engine/Plugins/Editor/ContentBrowser/ContentBrowserAssetDataSource/Source/ContentBrowserAssetDataSource/Private/ContentBrowserAssetDataCore.cpp
+++ b/Engine/Plugins/Editor/ContentBrowser/ContentBrowserAssetDataSource/Source/ContentBrowserAssetDataSource/Private/ContentBrowserAssetDataCore.cpp
@@ -13,7 +13,9 @@
 #include "ContentBrowserDataSource.h"
 #include "ContentBrowserDataSubsystem.h"
 #include "ContentBrowserDataUtils.h"
+#include "ContentBrowserItemPath.h"
 #include "ContentBrowserMenuContexts.h"
+#include "ContentBrowserModule.h"
 #include "Editor.h"
 #include "Editor/UnrealEdTypes.h"
 #include "Engine/Level.h"
@@ -918,13 +920,20 @@ bool PrivatizeItems(IAssetTools* InAssetTools, IAssetRegistry* InAssetRegistry,
 
 bool DeleteAssetFolderItems(TArrayView<const TSharedRef<const FContentBrowserAssetFolderItemDataPayload>> InFolderPayloads)
 {
-	TArray<FString> FoldersToDelete;
+	TArray<FContentBrowserItemPath> ItemPathsToDelete;
+	ItemPathsToDelete.Reserve(InFolderPayloads.Num());
 
 	for (const TSharedRef<const FContentBrowserAssetFolderItemDataPayload>& FolderPayload : InFolderPayloads)
 	{
-		FoldersToDelete.Add(FolderPayload->GetInternalPath().ToString());
+		ItemPathsToDelete.Add(FContentBrowserItemPath(FolderPayload->GetInternalPath(), EContentBrowserPathType::Internal));
 	}
 
+	FContentBrowserModule& ContentBrowserModule = FModuleManager::Get().GetModuleChecked<FContentBrowserModule>(TEXT("ContentBrowser"));
+	ContentBrowserModule.GetOnDeleteFolders().Broadcast(ItemPathsToDelete);
+
+	TArray<FString> FoldersToDelete;
+	Algo::Transform(ItemPathsToDelete, FoldersToDelete, [](FContentBrowserItemPath ItemPath) { return ItemPath.GetInternalPathString(); });
+
 	return FoldersToDelete.Num() > 0
 		&& AssetViewUtils::DeleteFolders(FoldersToDelete);
 }
diff --git a/Engine/Source/Editor/ContentBrowser/Public/ContentBrowserDelegates.h b/Engine/Source/Editor/ContentBrowser/Public/ContentBrowserDelegates.h
--- a/Engine/Source/Editor/ContentBrowser/Public/ContentBrowserDelegates.h
+++ b/Engine/Source/Editor/ContentBrowser/Public/ContentBrowserDelegates.h
@@ -79,6 +79,9 @@ DECLARE_DELEGATE_OneParam(FOnAssetEnterPressed, const TArray<FAssetData>& /*Sele
 /** Called when a new folder is starting to be created */
 DECLARE_DELEGATE_OneParam(FOnCreateNewFolder, const FContentBrowserItemTemporaryContext& /*NewItemContext*/);
 
+/** Called when one or more folders are being deleted */
+DECLARE_MULTICAST_DELEGATE_OneParam(FOnDeleteFolders, const TArray<FContentBrowserItemPath>& /*PathsToDelete*/);
+
 /** Called to request the menu when right clicking on an asset */
 DECLARE_DELEGATE_RetVal_OneParam(TSharedPtr<SWidget>, FOnGetAssetContextMenu, const TArray<FAssetData>& /*SelectedAssets*/);
 
diff --git a/Engine/Source/Editor/ContentBrowser/Public/ContentBrowserModule.h b/Engine/Source/Editor/ContentBrowser/Public/ContentBrowserModule.h
--- a/Engine/Source/Editor/ContentBrowser/Public/ContentBrowserModule.h
+++ b/Engine/Source/Editor/ContentBrowser/Public/ContentBrowserModule.h
@@ -187,6 +187,7 @@ public:
 	FOnSourcesViewChanged& GetOnSourcesViewChanged() { return OnSourcesViewChanged; }
 	FOnAssetPathChanged& GetOnAssetPathChanged() { return OnAssetPathChanged; }
 	FOnContentBrowserSettingChanged& GetOnContentBrowserSettingChanged() { return OnContentBrowserSettingChanged; }
+	FOnDeleteFolders& GetOnDeleteFolders() { return OnDeleteFolders; }
 
 	/** Override list of paths to select by default */
 	FDefaultSelectedPathsDelegate& GetDefaultSelectedPathsDelegate() { return DefaultSelectedPathsDelegate; }
@@ -259,6 +260,7 @@ private:
 	FOnSourcesViewChanged OnSourcesViewChanged;
 	FOnAssetPathChanged OnAssetPathChanged;
 	FOnContentBrowserSettingChanged OnContentBrowserSettingChanged;
+	FOnDeleteFolders OnDeleteFolders;
 
 	FDefaultSelectedPathsDelegate DefaultSelectedPathsDelegate;
 	FDefaultPathsToExpandDelegate DefaultPathsToExpandDelegate;
