DirectoryPlaceholder: Add content browser context menu action to cleanup (delete) placeholder assets from directories where they are no longer required, due to the presence of other valid asset types.
This function is also exposed to BP/Python.

diff --git a/Engine/Plugins/Developer/DirectoryPlaceholder/Source/DirectoryPlaceholder/DirectoryPlaceholder.build.cs b/Engine/Plugins/Developer/DirectoryPlaceholder/Source/DirectoryPlaceholder/DirectoryPlaceholder.build.cs
--- a/Engine/Plugins/Developer/DirectoryPlaceholder/Source/DirectoryPlaceholder/DirectoryPlaceholder.build.cs
+++ b/Engine/Plugins/Developer/DirectoryPlaceholder/Source/DirectoryPlaceholder/DirectoryPlaceholder.build.cs
@@ -21,6 +21,10 @@ namespace UnrealBuildTool.Rules
 					"Core",
 					"CoreUObject",
 					"DeveloperSettings",
+					"Engine",
+					"Slate",
+					"SlateCore",
+					"ToolMenus",
 					"UnrealEd",
 				}
 			);
diff --git a/Engine/Plugins/Developer/DirectoryPlaceholder/Source/DirectoryPlaceholder/Private/DirectoryPlaceholderModule.cpp b/Engine/Plugins/Developer/DirectoryPlaceholder/Source/DirectoryPlaceholder/Private/DirectoryPlaceholderModule.cpp
--- a/Engine/Plugins/Developer/DirectoryPlaceholder/Source/DirectoryPlaceholder/Private/DirectoryPlaceholderModule.cpp
+++ b/Engine/Plugins/Developer/DirectoryPlaceholder/Source/DirectoryPlaceholder/Private/DirectoryPlaceholderModule.cpp
@@ -5,14 +5,17 @@
 #include "AssetRegistry/AssetRegistryModule.h"
 #include "AssetToolsModule.h"
 #include "ContentBrowserItemPath.h"
+#include "ContentBrowserMenuContexts.h"
 #include "ContentBrowserModule.h"
 #include "DirectoryPlaceholder.h"
 #include "DirectoryPlaceholderFactory.h"
 #include "DirectoryPlaceholderSettings.h"
+#include "DirectoryPlaceholderUtils.h"
 #include "DirectoryWatcherModule.h"
 #include "FileHelpers.h"
 #include "IDirectoryWatcher.h"
 #include "ObjectTools.h"
+#include "ToolMenus.h"
 
 DEFINE_LOG_CATEGORY_STATIC(LogDirectoryPlaceholder, Log, All)
 
@@ -63,6 +66,30 @@ void FDirectoryPlaceholderModule::StartupModule()
 		// Register a callback to execute when one or more folders are being deleted in the content browser
 		FContentBrowserModule& ContentBrowserModule = FModuleManager::Get().GetModuleChecked<FContentBrowserModule>(TEXT("ContentBrowser"));
 		OnDeleteFoldersHandle = ContentBrowserModule.GetOnDeleteFolders().AddRaw(this, &FDirectoryPlaceholderModule::OnDeleteFolders);
+
+		// Extend the content browser folder context menu with the option to cleanup directory placeholders
+		if (UToolMenu* Menu = UToolMenus::Get()->ExtendMenu("ContentBrowser.FolderContextMenu"))
+		{
+			FToolMenuSection& BulkOpsSection = Menu->FindOrAddSection("PathContextBulkOperations");
+
+			FToolMenuEntry& Entry = BulkOpsSection.AddDynamicEntry("CleanupDirectoryPlaceholders", FNewToolMenuSectionDelegate::CreateLambda([this](FToolMenuSection& InSection)
+				{
+					if (UContentBrowserFolderContext* Context = InSection.FindContext<UContentBrowserFolderContext>())
+					{
+						const TArray<FString>& Paths = Context->GetSelectedPackagePaths();
+
+						FToolMenuEntry& Entry = InSection.AddMenuEntry(
+							"CleanupDirectoryPlaceholders",
+							LOCTEXT("CleanupDirectoryPlaceholdersLabel", "Cleanup Directory Placeholders"),
+							LOCTEXT("CleanupDirectoryPlaceholdersToolTip", "Delete all unnecessary placeholder assets in this folder (and sub-folders)"),
+							FSlateIcon(FAppStyle::GetAppStyleSetName(), "Icons.FolderClosed"),
+							FUIAction(FExecuteAction::CreateLambda([Paths]() { UDirectoryPlaceholderLibrary::CleanupPlaceholdersInPaths(Paths); }))
+						);
+					}
+				}));
+
+			Entry.InsertPosition = FToolMenuInsert("Delete", EToolMenuInsertType::After);
+		}
 	}
 }
 
@@ -132,7 +159,7 @@ void FDirectoryPlaceholderModule::OnDeleteFolders(const TArray<FContentBrowserIt
 		return;
 	}
 
-	FAssetRegistryModule& AssetRegistryModule = FModuleManager::Get().LoadModuleChecked<FAssetRegistryModule>(TEXT("AssetRegistry"));
+	FAssetRegistryModule& AssetRegistryModule = FModuleManager::Get().GetModuleChecked<FAssetRegistryModule>(TEXT("AssetRegistry"));
 
 	// Set up a recursive asset registry filter for the paths being deleted
 	FARFilter Filter;
diff --git a/Engine/Plugins/Developer/DirectoryPlaceholder/Source/DirectoryPlaceholder/Private/DirectoryPlaceholderUtils.cpp b/Engine/Plugins/Developer/DirectoryPlaceholder/Source/DirectoryPlaceholder/Private/DirectoryPlaceholderUtils.cpp
new file mode 100644
--- /dev/null
+++ b/Engine/Plugins/Developer/DirectoryPlaceholder/Source/DirectoryPlaceholder/Private/DirectoryPlaceholderUtils.cpp
@@ -0,0 +1,78 @@
+// Copyright Epic Games, Inc. All Rights Reserved.
+
+#include "DirectoryPlaceholderUtils.h"
+
+#include "AssetRegistry/AssetRegistryModule.h"
+#include "DirectoryPlaceholder.h"
+#include "Misc/ScopedSlowTask.h"
+#include "Modules/ModuleManager.h"
+#include "ObjectTools.h"
+
+#define LOCTEXT_NAMESPACE "DirectoryPlaceholderLibrary"
+
+namespace UE::DirectoryPlaceholder::Private
+{
+	/** Recurisvely delete all unnecessary placeholder assets in this folder (and sub-folders) */
+	bool CleanupPlaceholdersInternal(FName Path)
+	{
+		FAssetRegistryModule& AssetRegistryModule = FModuleManager::Get().GetModuleChecked<FAssetRegistryModule>(TEXT("AssetRegistry"));
+
+		// Find all of the subfolders in this folder (non-recursively)
+		TArray<FName> SubPaths;
+		constexpr bool bRecursive = false;
+		AssetRegistryModule.Get().GetSubPaths(Path, SubPaths, bRecursive);
+
+		// Recursively check each of the subpaths for non-placeholder assets
+		bool bSubPathsHaveAssets = false;
+		FScopedSlowTask Progress(SubPaths.Num(), LOCTEXT("SlowTaskStartText", "Removing unnecessary placeholders..."));
+		for (FName SubPath : SubPaths)
+		{
+			Progress.EnterProgressFrame(1.0f, FText::Format(LOCTEXT("SlowTaskSubText", "{0}"), FText::FromName(SubPath)));
+			bSubPathsHaveAssets |= CleanupPlaceholdersInternal(SubPath);
+		}
+
+		// Get all of the assets in the current folder (non-recursive)
+		TArray<FAssetData> AssetDataList;
+		AssetRegistryModule.Get().GetAssetsByPath(Path, AssetDataList, bRecursive);
+
+		// Test if there are any assets in the current path that are NOT placeholders
+		const bool bHasAssets = AssetDataList.ContainsByPredicate([](const FAssetData& AssetData) { return AssetData.GetClass() != UDirectoryPlaceholder::StaticClass(); });
+
+		// If there is at least one non-placeholder asset in this folder, or one of its subfolders, then we can safely delete the placeholder(s) in this folder
+		if (bHasAssets || bSubPathsHaveAssets)
+		{
+			TArray<FAssetData> PlacesholdersToDelete = AssetDataList.FilterByPredicate([](const FAssetData& AssetData) {return AssetData.GetClass() == UDirectoryPlaceholder::StaticClass(); });
+
+			if (PlacesholdersToDelete.Num() > 0)
+			{
+				const bool bShowConfirmation = false;
+				const int32 NumAssetsDeleted = ObjectTools::DeleteAssets(PlacesholdersToDelete, bShowConfirmation);
+			}
+		}
+
+		return bHasAssets || bSubPathsHaveAssets;
+	}
+}
+
+void UDirectoryPlaceholderLibrary::CleanupPlaceholdersInPath(const FString& Path)
+{
+	TArray<FString> Paths = { Path };
+	CleanupPlaceholdersInPaths(Paths);
+}
+
+void UDirectoryPlaceholderLibrary::CleanupPlaceholdersInPaths(const TArray<FString>& Paths)
+{
+	FAssetRegistryModule& AssetRegistryModule = FModuleManager::Get().GetModuleChecked<FAssetRegistryModule>(TEXT("AssetRegistry"));
+
+	TArray<FString> ValidPaths = Paths.FilterByPredicate([&AssetRegistryModule](FString Path) { return AssetRegistryModule.Get().PathExists(Path); });
+
+	FScopedSlowTask Progress(ValidPaths.Num(), LOCTEXT("SlowTaskStartText", "Removing unnecessary placeholders..."));
+	Progress.MakeDialog();
+	for (const FString& Path : ValidPaths)
+	{
+		Progress.EnterProgressFrame();
+		UE::DirectoryPlaceholder::Private::CleanupPlaceholdersInternal(*Path);
+	}
+}
+
+#undef LOCTEXT_NAMESPACE
diff --git a/Engine/Plugins/Developer/DirectoryPlaceholder/Source/DirectoryPlaceholder/Private/DirectoryPlaceholderUtils.h b/Engine/Plugins/Developer/DirectoryPlaceholder/Source/DirectoryPlaceholder/Private/DirectoryPlaceholderUtils.h
new file mode 100644
--- /dev/null
+++ b/Engine/Plugins/Developer/DirectoryPlaceholder/Source/DirectoryPlaceholder/Private/DirectoryPlaceholderUtils.h
@@ -0,0 +1,25 @@
+// Copyright Epic Games, Inc. All Rights Reserved.
+
+#pragma once
+
+#include "Kismet/BlueprintFunctionLibrary.h"
+
+#include "DirectoryPlaceholderUtils.generated.h"
+
+/**
+ * Library functions for operations on directory placeholder assets
+ */
+UCLASS()
+class UDirectoryPlaceholderLibrary : public UBlueprintFunctionLibrary
+{
+	GENERATED_BODY()
+
+public:
+	/** Delete all unnecessary placeholder assets in this folder (and sub-folders) */
+	UFUNCTION(BlueprintCallable, Category = "DirectoryPlaceholder")
+	static void CleanupPlaceholdersInPath(const FString& Path);
+
+	/** Delete all unnecessary placeholder assets in these folders (and sub-folders) */
+	UFUNCTION(BlueprintCallable, Category = "DirectoryPlaceholder")
+	static void CleanupPlaceholdersInPaths(const TArray<FString>& Paths);
+};
