DirectoryPlaceholder: Adds check for existing placeholders to prevent the creation of unnecessary additional placeholder assets. This also fixes a bug with folders getting renamed / copied / moved.
#rnx

diff --git a/Engine/Plugins/Developer/DirectoryPlaceholder/Source/DirectoryPlaceholder/Private/DirectoryPlaceholderModule.cpp b/Engine/Plugins/Developer/DirectoryPlaceholder/Source/DirectoryPlaceholder/Private/DirectoryPlaceholderModule.cpp
--- a/Engine/Plugins/Developer/DirectoryPlaceholder/Source/DirectoryPlaceholder/Private/DirectoryPlaceholderModule.cpp
+++ b/Engine/Plugins/Developer/DirectoryPlaceholder/Source/DirectoryPlaceholder/Private/DirectoryPlaceholderModule.cpp
@@ -121,9 +121,24 @@ void FDirectoryPlaceholderModule::OnDirectoryChanged(const TArray<FFileChangeDat
 		// Check if the FileChangeData represent a directory that was just added
 		if (FPaths::DirectoryExists(FileChange.Filename) && (FileChange.Action == FFileChangeData::FCA_Added))
 		{
+			const FString PackagePath = FPackageName::FilenameToLongPackageName(FileChange.Filename);
+
+			// Test if there are any placeholder assets already in the path of the new directory. If there are, we will not make a new one.
+			// Note: This could occur if the directory was just renamed, copied, or moved.
+			FAssetRegistryModule& AssetRegistryModule = FModuleManager::Get().GetModuleChecked<FAssetRegistryModule>(TEXT("AssetRegistry"));
+
+			TArray<FAssetData> AssetDataList;
+			constexpr bool bRecursive = false;
+			AssetRegistryModule.Get().GetAssetsByPath(*PackagePath, AssetDataList, bRecursive);
+
+			const bool bHasPlaceholders = AssetDataList.ContainsByPredicate([](const FAssetData& AssetData) { return AssetData.GetClass() == UDirectoryPlaceholder::StaticClass(); });
+			if (bHasPlaceholders)
+			{
+				continue;
+			}
+
 			// Create a new UDirectoryPlaceholder asset in the directory that was just added
 			const FString AssetName = FString(TEXT("UE_Placeholder"));
-			const FString PackagePath = FPackageName::FilenameToLongPackageName(FileChange.Filename);
 
 			IAssetTools& AssetTools = FModuleManager::GetModuleChecked<FAssetToolsModule>("AssetTools").Get();
 			const UObject* const PlaceholderAsset = AssetTools.CreateAsset(AssetName, PackagePath, UDirectoryPlaceholder::StaticClass(), NewObject<UDirectoryPlaceholderFactory>());
