DirectoryPlaceholder: Various bugfixes
- Check for any existing assets (not just placeholders) before creating a placeholder in a "new folder". This prevents unecessary placeholders from being added when a folder is renamed / moved / copied. As a side effect, this also fixes an issue where placeholders were being created in __ExternalActors__ and __ExternalObjects__ folders.
- Remove "/All" from paths before attempting to cleanup placeholders, because the asset registry does not recognize paths with that prefix as being valid.
- Only automatically delete / cleanup placeholders from folders in project content ("/Game"), because that is the only place we automatically create placeholders.
- Prevent an asset editor from being opened for directory placeholders.
- Only add "cleanup" context menu action for folders with at least one asset path.

#rnx

diff --git a/Engine/Plugins/Developer/DirectoryPlaceholder/Source/DirectoryPlaceholder/Private/AssetDefinition_DirectoryPlaceholder.cpp b/Engine/Plugins/Developer/DirectoryPlaceholder/Source/DirectoryPlaceholder/Private/AssetDefinition_DirectoryPlaceholder.cpp
--- a/Engine/Plugins/Developer/DirectoryPlaceholder/Source/DirectoryPlaceholder/Private/AssetDefinition_DirectoryPlaceholder.cpp
+++ b/Engine/Plugins/Developer/DirectoryPlaceholder/Source/DirectoryPlaceholder/Private/AssetDefinition_DirectoryPlaceholder.cpp
@@ -26,4 +26,9 @@ TConstArrayView<FAssetCategoryPath> UAssetDefinition_DirectoryPlaceholder::GetAs
 	return Categories;
 }
 
+EAssetCommandResult UAssetDefinition_DirectoryPlaceholder::OpenAssets(const FAssetOpenArgs& OpenArgs) const
+{
+	return EAssetCommandResult::Unhandled;
+}
+
 #undef LOCTEXT_NAMESPACE
diff --git a/Engine/Plugins/Developer/DirectoryPlaceholder/Source/DirectoryPlaceholder/Private/AssetDefinition_DirectoryPlaceholder.h b/Engine/Plugins/Developer/DirectoryPlaceholder/Source/DirectoryPlaceholder/Private/AssetDefinition_DirectoryPlaceholder.h
--- a/Engine/Plugins/Developer/DirectoryPlaceholder/Source/DirectoryPlaceholder/Private/AssetDefinition_DirectoryPlaceholder.h
+++ b/Engine/Plugins/Developer/DirectoryPlaceholder/Source/DirectoryPlaceholder/Private/AssetDefinition_DirectoryPlaceholder.h
@@ -18,5 +18,6 @@ protected:
 	virtual FLinearColor GetAssetColor() const override;
 	virtual TSoftClassPtr<UObject> GetAssetClass() const override;
 	virtual TConstArrayView<FAssetCategoryPath> GetAssetCategories() const override;
+	virtual EAssetCommandResult OpenAssets(const FAssetOpenArgs& OpenArgs) const override;
 	// End UAssetDefinition Interface
 };
diff --git a/Engine/Plugins/Developer/DirectoryPlaceholder/Source/DirectoryPlaceholder/Private/DirectoryPlaceholderModule.cpp b/Engine/Plugins/Developer/DirectoryPlaceholder/Source/DirectoryPlaceholder/Private/DirectoryPlaceholderModule.cpp
--- a/Engine/Plugins/Developer/DirectoryPlaceholder/Source/DirectoryPlaceholder/Private/DirectoryPlaceholderModule.cpp
+++ b/Engine/Plugins/Developer/DirectoryPlaceholder/Source/DirectoryPlaceholder/Private/DirectoryPlaceholderModule.cpp
@@ -74,7 +74,8 @@ void FDirectoryPlaceholderModule::StartupModule()
 
 			FToolMenuEntry& Entry = BulkOpsSection.AddDynamicEntry("CleanupDirectoryPlaceholders", FNewToolMenuSectionDelegate::CreateLambda([this](FToolMenuSection& InSection)
 				{
-					if (UContentBrowserFolderContext* Context = InSection.FindContext<UContentBrowserFolderContext>())
+					UContentBrowserFolderContext* Context = InSection.FindContext<UContentBrowserFolderContext>();
+					if (Context && Context->NumAssetPaths > 0)
 					{
 						const TArray<FString>& Paths = Context->GetSelectedPackagePaths();
 
@@ -85,10 +86,10 @@ void FDirectoryPlaceholderModule::StartupModule()
 							FSlateIcon(FAppStyle::GetAppStyleSetName(), "Icons.FolderClosed"),
 							FUIAction(FExecuteAction::CreateLambda([Paths]() { UDirectoryPlaceholderLibrary::CleanupPlaceholdersInPaths(Paths); }))
 						);
+
+						Entry.InsertPosition = FToolMenuInsert("Delete", EToolMenuInsertType::After);
 					}
 				}));
-
-			Entry.InsertPosition = FToolMenuInsert("Delete", EToolMenuInsertType::After);
 		}
 	}
 }
@@ -123,16 +124,13 @@ void FDirectoryPlaceholderModule::OnDirectoryChanged(const TArray<FFileChangeDat
 		{
 			const FString PackagePath = FPackageName::FilenameToLongPackageName(FileChange.Filename);
 
-			// Test if there are any placeholder assets already in the path of the new directory. If there are, we will not make a new one.
+			// Test if there are any assets already in the path of the new directory. If there are, we do not create a placeholder.
 			// Note: This could occur if the directory was just renamed, copied, or moved.
 			FAssetRegistryModule& AssetRegistryModule = FModuleManager::Get().GetModuleChecked<FAssetRegistryModule>(TEXT("AssetRegistry"));
 
 			TArray<FAssetData> AssetDataList;
-			constexpr bool bRecursive = false;
-			AssetRegistryModule.Get().GetAssetsByPath(*PackagePath, AssetDataList, bRecursive);
-
-			const bool bHasPlaceholders = AssetDataList.ContainsByPredicate([](const FAssetData& AssetData) { return AssetData.GetClass() == UDirectoryPlaceholder::StaticClass(); });
-			if (bHasPlaceholders)
+			constexpr bool bRecursive = true;
+			if (AssetRegistryModule.Get().HasAssets(*PackagePath, bRecursive))
 			{
 				continue;
 			}
@@ -176,12 +174,22 @@ void FDirectoryPlaceholderModule::OnDeleteFolders(const TArray<FContentBrowserIt
 
 	FAssetRegistryModule& AssetRegistryModule = FModuleManager::Get().GetModuleChecked<FAssetRegistryModule>(TEXT("AssetRegistry"));
 
-	// Set up a recursive asset registry filter for the paths being deleted
+	// The only place we automatically create directory placeholders is under project content ("/Game").
+	// These are also the only paths that we want to automatically delete placeholders from.
 	FARFilter Filter;
 	Filter.bRecursivePaths = true;
 	for (const FContentBrowserItemPath& Path : PathsToDelete)
 	{
-		Filter.PackagePaths.Add(Path.GetInternalPathName());
+		FString InternalPath = Path.GetInternalPathName().ToString();
+		if (InternalPath.Equals(TEXT("/Game")) || InternalPath.StartsWith(TEXT("/Game/")))
+		{
+			Filter.PackagePaths.Add(Path.GetInternalPathName());
+		}
+	}
+
+	if (Filter.PackagePaths.IsEmpty())
+	{
+		return;
 	}
 
 	// Find all of the assets in the folders being deleted
diff --git a/Engine/Plugins/Developer/DirectoryPlaceholder/Source/DirectoryPlaceholder/Private/DirectoryPlaceholderUtils.cpp b/Engine/Plugins/Developer/DirectoryPlaceholder/Source/DirectoryPlaceholder/Private/DirectoryPlaceholderUtils.cpp
--- a/Engine/Plugins/Developer/DirectoryPlaceholder/Source/DirectoryPlaceholder/Private/DirectoryPlaceholderUtils.cpp
+++ b/Engine/Plugins/Developer/DirectoryPlaceholder/Source/DirectoryPlaceholder/Private/DirectoryPlaceholderUtils.cpp
@@ -64,7 +64,20 @@ void UDirectoryPlaceholderLibrary::CleanupPlaceholdersInPaths(const TArray<FStri
 {
 	FAssetRegistryModule& AssetRegistryModule = FModuleManager::Get().GetModuleChecked<FAssetRegistryModule>(TEXT("AssetRegistry"));
 
-	TArray<FString> ValidPaths = Paths.FilterByPredicate([&AssetRegistryModule](FString Path) { return AssetRegistryModule.Get().PathExists(Path); });
+	// The asset registry will not recognize paths that start with "/All", so remove that prefix if present
+	TArray<FString> ValidPaths;
+	Algo::Transform(Paths, ValidPaths, [](FString Path) 
+		{ 
+			Path.RemoveFromStart(TEXT("/All")); 
+			return Path; 
+		});
+
+	// The only place we automatically create directory placeholders is under project content ("/Game").
+	// These are also the only paths that we want to automatically delete placeholders from.
+	ValidPaths = ValidPaths.FilterByPredicate([&AssetRegistryModule](FString Path)
+		{ 
+			return AssetRegistryModule.Get().PathExists(Path) && (Path.Equals(TEXT("/Game")) || Path.StartsWith(TEXT("/Game/")));
+		});
 
 	FScopedSlowTask Progress(ValidPaths.Num(), LOCTEXT("SlowTaskStartText", "Removing unnecessary placeholders..."));
 	Progress.MakeDialog();
