DirectoryPlaceholders: Add directory watcher to ProjectPluginsDir to enable automatic creation of placeholders in folders in Project Plugin Content
#rnx

diff --git a/Engine/Plugins/Developer/DirectoryPlaceholder/Source/DirectoryPlaceholder/Private/DirectoryPlaceholderModule.cpp b/Engine/Plugins/Developer/DirectoryPlaceholder/Source/DirectoryPlaceholder/Private/DirectoryPlaceholderModule.cpp
--- a/Engine/Plugins/Developer/DirectoryPlaceholder/Source/DirectoryPlaceholder/Private/DirectoryPlaceholderModule.cpp
+++ b/Engine/Plugins/Developer/DirectoryPlaceholder/Source/DirectoryPlaceholder/Private/DirectoryPlaceholderModule.cpp
@@ -61,6 +61,13 @@ void FDirectoryPlaceholderModule::StartupModule()
 				EngineDirectoryWatcherHandle,
 				IDirectoryWatcher::WatchOptions::IncludeDirectoryChanges
 			);
+
+			DirectoryWatcher->RegisterDirectoryChangedCallback_Handle(
+				FPaths::ProjectPluginsDir(),
+				IDirectoryWatcher::FDirectoryChanged::CreateRaw(this, &FDirectoryPlaceholderModule::OnDirectoryChanged),
+				EngineDirectoryWatcherHandle,
+				IDirectoryWatcher::WatchOptions::IncludeDirectoryChanges
+			);
 		}
 
 		// Register a callback to execute when one or more folders are being deleted in the content browser
@@ -122,7 +129,11 @@ void FDirectoryPlaceholderModule::OnDirectoryChanged(const TArray<FFileChangeDat
 		// Check if the FileChangeData represent a directory that was just added
 		if (FPaths::DirectoryExists(FileChange.Filename) && (FileChange.Action == FFileChangeData::FCA_Added))
 		{
-			const FString PackagePath = FPackageName::FilenameToLongPackageName(FileChange.Filename);
+			FString PackagePath;
+			if (!FPackageName::TryConvertFilenameToLongPackageName(FileChange.Filename, PackagePath))
+			{
+				continue;
+			}
 
 			// Test if there are any assets already in the path of the new directory. If there are, we do not create a placeholder.
 			// Note: This could occur if the directory was just renamed, copied, or moved.
