DirectoryPlaceholders: Delay registration of the directory watchers until after the asset registry is finished loading.
#rnx

diff --git a/Engine/Plugins/Developer/DirectoryPlaceholder/Source/DirectoryPlaceholder/Private/DirectoryPlaceholderModule.cpp b/Engine/Plugins/Developer/DirectoryPlaceholder/Source/DirectoryPlaceholder/Private/DirectoryPlaceholderModule.cpp
--- a/Engine/Plugins/Developer/DirectoryPlaceholder/Source/DirectoryPlaceholder/Private/DirectoryPlaceholderModule.cpp
+++ b/Engine/Plugins/Developer/DirectoryPlaceholder/Source/DirectoryPlaceholder/Private/DirectoryPlaceholderModule.cpp
@@ -35,6 +35,9 @@ public:
 	//~ End IModuleInterface interface
 
 private:
+	/** Register the Directory Watchers */
+	void RegisterDirectoryWatchers();
+
 	/** Callback when the directory watcher detects a file/folder change */
 	void OnDirectoryChanged(const TArray<FFileChangeData>& InFileChanges);
 
@@ -51,24 +54,9 @@ void FDirectoryPlaceholderModule::StartupModule()
 	// Only enable this behavior in an interactive editor
 	if (GIsEditor && !IsRunningCommandlet())
 	{
-		// Register a callback with the directory watcher to be notified about file/folder changes
-		FDirectoryWatcherModule& DirectoryWatcherModule = FModuleManager::GetModuleChecked<FDirectoryWatcherModule>("DirectoryWatcher");
-		if (IDirectoryWatcher* DirectoryWatcher = DirectoryWatcherModule.Get())
-		{
-			DirectoryWatcher->RegisterDirectoryChangedCallback_Handle(
-				FPaths::ProjectContentDir(),
-				IDirectoryWatcher::FDirectoryChanged::CreateRaw(this, &FDirectoryPlaceholderModule::OnDirectoryChanged),
-				EngineDirectoryWatcherHandle,
-				IDirectoryWatcher::WatchOptions::IncludeDirectoryChanges
-			);
-
-			DirectoryWatcher->RegisterDirectoryChangedCallback_Handle(
-				FPaths::ProjectPluginsDir(),
-				IDirectoryWatcher::FDirectoryChanged::CreateRaw(this, &FDirectoryPlaceholderModule::OnDirectoryChanged),
-				EngineDirectoryWatcherHandle,
-				IDirectoryWatcher::WatchOptions::IncludeDirectoryChanges
-			);
-		}
+		// Delay registration of the directory watchers until after the asset registry has finished its initial scan
+		const FAssetRegistryModule& AssetRegistryModule = FModuleManager::LoadModuleChecked<FAssetRegistryModule>("AssetRegistry");
+		AssetRegistryModule.Get().OnFilesLoaded().AddRaw(this, &FDirectoryPlaceholderModule::RegisterDirectoryWatchers);
 
 		// Register a callback to execute when one or more folders are being deleted in the content browser
 		FContentBrowserModule& ContentBrowserModule = FModuleManager::Get().GetModuleChecked<FContentBrowserModule>(TEXT("ContentBrowser"));
@@ -116,6 +104,28 @@ void FDirectoryPlaceholderModule::ShutdownModule()
 	}
 }
 
+void FDirectoryPlaceholderModule::RegisterDirectoryWatchers()
+{
+	// Register a callback with the directory watcher to be notified about file/folder changes
+	FDirectoryWatcherModule& DirectoryWatcherModule = FModuleManager::GetModuleChecked<FDirectoryWatcherModule>("DirectoryWatcher");
+	if (IDirectoryWatcher* DirectoryWatcher = DirectoryWatcherModule.Get())
+	{
+		DirectoryWatcher->RegisterDirectoryChangedCallback_Handle(
+			FPaths::ProjectContentDir(),
+			IDirectoryWatcher::FDirectoryChanged::CreateRaw(this, &FDirectoryPlaceholderModule::OnDirectoryChanged),
+			EngineDirectoryWatcherHandle,
+			IDirectoryWatcher::WatchOptions::IncludeDirectoryChanges
+		);
+
+		DirectoryWatcher->RegisterDirectoryChangedCallback_Handle(
+			FPaths::ProjectPluginsDir(),
+			IDirectoryWatcher::FDirectoryChanged::CreateRaw(this, &FDirectoryPlaceholderModule::OnDirectoryChanged),
+			EngineDirectoryWatcherHandle,
+			IDirectoryWatcher::WatchOptions::IncludeDirectoryChanges
+		);
+	}
+}
+
 void FDirectoryPlaceholderModule::OnDirectoryChanged(const TArray<FFileChangeData>& InFileChanges)
 {
 	const UDirectoryPlaceholderSettings* const Settings = GetDefault<UDirectoryPlaceholderSettings>();
