DirectoryPlaceholders: Fix error about missing project plugin directory when trying to register a directory watcher.
#rnx

diff --git a/Engine/Plugins/Developer/DirectoryPlaceholder/Source/DirectoryPlaceholder/Private/DirectoryPlaceholderModule.cpp b/Engine/Plugins/Developer/DirectoryPlaceholder/Source/DirectoryPlaceholder/Private/DirectoryPlaceholderModule.cpp
--- a/Engine/Plugins/Developer/DirectoryPlaceholder/Source/DirectoryPlaceholder/Private/DirectoryPlaceholderModule.cpp
+++ b/Engine/Plugins/Developer/DirectoryPlaceholder/Source/DirectoryPlaceholder/Private/DirectoryPlaceholderModule.cpp
@@ -110,19 +110,25 @@ void FDirectoryPlaceholderModule::RegisterDirectoryWatchers()
 	FDirectoryWatcherModule& DirectoryWatcherModule = FModuleManager::GetModuleChecked<FDirectoryWatcherModule>("DirectoryWatcher");
 	if (IDirectoryWatcher* DirectoryWatcher = DirectoryWatcherModule.Get())
 	{
-		DirectoryWatcher->RegisterDirectoryChangedCallback_Handle(
-			FPaths::ProjectContentDir(),
-			IDirectoryWatcher::FDirectoryChanged::CreateRaw(this, &FDirectoryPlaceholderModule::OnDirectoryChanged),
-			EngineDirectoryWatcherHandle,
-			IDirectoryWatcher::WatchOptions::IncludeDirectoryChanges
-		);
-
-		DirectoryWatcher->RegisterDirectoryChangedCallback_Handle(
-			FPaths::ProjectPluginsDir(),
-			IDirectoryWatcher::FDirectoryChanged::CreateRaw(this, &FDirectoryPlaceholderModule::OnDirectoryChanged),
-			EngineDirectoryWatcherHandle,
-			IDirectoryWatcher::WatchOptions::IncludeDirectoryChanges
-		);
+		if (FPaths::DirectoryExists(FPaths::ProjectContentDir()))
+		{
+			DirectoryWatcher->RegisterDirectoryChangedCallback_Handle(
+				FPaths::ProjectContentDir(),
+				IDirectoryWatcher::FDirectoryChanged::CreateRaw(this, &FDirectoryPlaceholderModule::OnDirectoryChanged),
+				EngineDirectoryWatcherHandle,
+				IDirectoryWatcher::WatchOptions::IncludeDirectoryChanges
+			);
+		}
+
+		if (FPaths::DirectoryExists(FPaths::ProjectPluginsDir()))
+		{
+			DirectoryWatcher->RegisterDirectoryChangedCallback_Handle(
+				FPaths::ProjectPluginsDir(),
+				IDirectoryWatcher::FDirectoryChanged::CreateRaw(this, &FDirectoryPlaceholderModule::OnDirectoryChanged),
+				EngineDirectoryWatcherHandle,
+				IDirectoryWatcher::WatchOptions::IncludeDirectoryChanges
+			);
+		}
 	}
 }
 
