DirectoryPlaceholders: Fix crash when trying to save a large number of placeholder assets one at a time by making a single call to SavePackages with all of the placeholders being created.
#rnx

diff --git a/Engine/Plugins/Developer/DirectoryPlaceholder/Source/DirectoryPlaceholder/Private/DirectoryPlaceholderModule.cpp b/Engine/Plugins/Developer/DirectoryPlaceholder/Source/DirectoryPlaceholder/Private/DirectoryPlaceholderModule.cpp
--- a/Engine/Plugins/Developer/DirectoryPlaceholder/Source/DirectoryPlaceholder/Private/DirectoryPlaceholderModule.cpp
+++ b/Engine/Plugins/Developer/DirectoryPlaceholder/Source/DirectoryPlaceholder/Private/DirectoryPlaceholderModule.cpp
@@ -140,6 +140,8 @@ void FDirectoryPlaceholderModule::OnDirectoryChanged(const TArray<FFileChangeDat
 		return;
 	}
 
+	TArray<UPackage*> PackagesToSave;
+
 	for (const FFileChangeData& FileChange : InFileChanges)
 	{
 		// Check if the FileChangeData represent a directory that was just added
@@ -172,16 +174,8 @@ void FDirectoryPlaceholderModule::OnDirectoryChanged(const TArray<FFileChangeDat
 			// If source control is enabled, it will also be marked for add.
 			if (PlaceholderAsset)
 			{
-				TArray<UPackage*> PackagesToSave = { PlaceholderAsset->GetPackage() };
-				const bool bOnlyDirty = false;
-				if (UEditorLoadingAndSavingUtils::SavePackages(PackagesToSave, bOnlyDirty))
-				{
-					UE_LOG(LogDirectoryPlaceholder, Log, TEXT("New Directory Placeholder was created and saved in %s"), *PackagePath);
-				}
-				else
-				{
-					UE_LOG(LogDirectoryPlaceholder, Log, TEXT("New Directory Placeholder was created in %s, but could not be saved."), *PackagePath);
-				}
+				PackagesToSave.Add(PlaceholderAsset->GetPackage());
+				UE_LOG(LogDirectoryPlaceholder, Verbose, TEXT("New Directory Placeholder was created in %s"), *PackagePath);
 			}
 			else
 			{
@@ -189,6 +183,9 @@ void FDirectoryPlaceholderModule::OnDirectoryChanged(const TArray<FFileChangeDat
 			}
 		}
 	}
+
+	const bool bOnlyDirty = false;
+	UEditorLoadingAndSavingUtils::SavePackages(PackagesToSave, bOnlyDirty);
 }
 
 void FDirectoryPlaceholderModule::OnDeleteFolders(const TArray<FContentBrowserItemPath>& PathsToDelete)
