DirectoryPlaceholders: Fix issue where placeholder assets could get created in ExternalActors and ExternalObjects folders.
Adds project setting for a list of exclude folders where users do not want directory placeholders to be automatically created.
#rnx

diff --git a/Engine/Plugins/Developer/DirectoryPlaceholder/Source/DirectoryPlaceholder/Private/DirectoryPlaceholderModule.cpp b/Engine/Plugins/Developer/DirectoryPlaceholder/Source/DirectoryPlaceholder/Private/DirectoryPlaceholderModule.cpp
--- a/Engine/Plugins/Developer/DirectoryPlaceholder/Source/DirectoryPlaceholder/Private/DirectoryPlaceholderModule.cpp
+++ b/Engine/Plugins/Developer/DirectoryPlaceholder/Source/DirectoryPlaceholder/Private/DirectoryPlaceholderModule.cpp
@@ -2,6 +2,7 @@
 
 #include "Modules/ModuleManager.h"
 
+#include "Algo/AnyOf.h"
 #include "AssetRegistry/AssetRegistryModule.h"
 #include "AssetToolsModule.h"
 #include "ContentBrowserItemPath.h"
@@ -153,6 +154,23 @@ void FDirectoryPlaceholderModule::OnDirectoryChanged(const TArray<FFileChangeDat
 				continue;
 			}
 
+			// Do not create directory placeholders in the ExternalActors of ExternalObjects folders
+			if (PackagePath.Contains(FPackagePath::GetExternalActorsFolderName()) || PackagePath.Contains(FPackagePath::GetExternalObjectsFolderName()))
+			{
+				continue;
+			}
+
+			// Do not create directory placeholders in any folders excluded by the project settings
+			const bool bExclude = Algo::AnyOf(Settings->ExcludePaths, [&PackagePath](const FString& ExcludePath)
+				{
+					return FPaths::IsUnderDirectory(PackagePath, ExcludePath);
+				});
+
+			if (bExclude)
+			{
+				continue;
+			}
+
 			// Test if there are any assets already in the path of the new directory. If there are, we do not create a placeholder.
 			// Note: This could occur if the directory was just renamed, copied, or moved.
 			FAssetRegistryModule& AssetRegistryModule = FModuleManager::Get().GetModuleChecked<FAssetRegistryModule>(TEXT("AssetRegistry"));
diff --git a/Engine/Plugins/Developer/DirectoryPlaceholder/Source/DirectoryPlaceholder/Private/DirectoryPlaceholderSettings.h b/Engine/Plugins/Developer/DirectoryPlaceholder/Source/DirectoryPlaceholder/Private/DirectoryPlaceholderSettings.h
--- a/Engine/Plugins/Developer/DirectoryPlaceholder/Source/DirectoryPlaceholder/Private/DirectoryPlaceholderSettings.h
+++ b/Engine/Plugins/Developer/DirectoryPlaceholder/Source/DirectoryPlaceholder/Private/DirectoryPlaceholderSettings.h
@@ -18,6 +18,13 @@ public:
 	UDirectoryPlaceholderSettings() = default;
 
 	/** If enabled, Directory Placeholder assets will be automatically created in new folders, and automatically deleted from folders being deleted */
-	UPROPERTY(EditAnywhere, Category = "Default")
+	UPROPERTY(Config, EditAnywhere, Category = "Default")
 	bool bAutomaticallyCreatePlaceholders = true;
+
+	/** 
+	 * Directory Placeholder assets will not be automatically created in new folders under any of the paths in this list.
+	 * Use content browser paths, such as "/Game/MyFolder/" or "/MyPlugin/MyFolder/".
+	 */
+	UPROPERTY(Config, EditAnywhere, Category = "Default")
+	TArray<FString> ExcludePaths;
 };
