DirectoryPlaceholders: Add support to automatically create placeholders in new folders within additional plugin directories.
Add additional editor user settings to better control where placeholders can be automatically created.
Fix unbalanced register/unregister of existing directory watcher delegates.
Update OnDeleteFolders() to allow for placeholdesr to be automatically deleted from any folder (not just project content) if they are the only assets that exist in the folder being deleted.

#rnx

diff --git a/Engine/Plugins/Developer/DirectoryPlaceholder/Source/DirectoryPlaceholder/DirectoryPlaceholder.build.cs b/Engine/Plugins/Developer/DirectoryPlaceholder/Source/DirectoryPlaceholder/DirectoryPlaceholder.build.cs
--- a/Engine/Plugins/Developer/DirectoryPlaceholder/Source/DirectoryPlaceholder/DirectoryPlaceholder.build.cs
+++ b/Engine/Plugins/Developer/DirectoryPlaceholder/Source/DirectoryPlaceholder/DirectoryPlaceholder.build.cs
@@ -22,6 +22,7 @@ namespace UnrealBuildTool.Rules
 					"CoreUObject",
 					"DeveloperSettings",
 					"Engine",
+					"Projects",
 					"Slate",
 					"SlateCore",
 					"ToolMenus",
diff --git a/Engine/Plugins/Developer/DirectoryPlaceholder/Source/DirectoryPlaceholder/Private/DirectoryPlaceholderModule.cpp b/Engine/Plugins/Developer/DirectoryPlaceholder/Source/DirectoryPlaceholder/Private/DirectoryPlaceholderModule.cpp
--- a/Engine/Plugins/Developer/DirectoryPlaceholder/Source/DirectoryPlaceholder/Private/DirectoryPlaceholderModule.cpp
+++ b/Engine/Plugins/Developer/DirectoryPlaceholder/Source/DirectoryPlaceholder/Private/DirectoryPlaceholderModule.cpp
@@ -15,6 +15,7 @@
 #include "DirectoryWatcherModule.h"
 #include "FileHelpers.h"
 #include "IDirectoryWatcher.h"
+#include "Interfaces/IProjectManager.h"
 #include "ObjectTools.h"
 #include "ToolMenus.h"
 
@@ -36,8 +37,17 @@ public:
 	//~ End IModuleInterface interface
 
 private:
-	/** Register the Directory Watchers */
+	/** Register/Unregister all enabled Directory Watchers */
 	void RegisterDirectoryWatchers();
+	void UnregisterDirectoryWatchers();
+
+	/** Register/Unregister a Directory Watcher for the input path */
+	void RegisterDirectoryWatcher(const FString& Path, FDelegateHandle& Handle);
+	void UnregisterDirectoryWatcher(const FString& Path, FDelegateHandle& Handle);
+
+	/** Register/Unregister a Directory Watcher for each additional (external) plugin path */
+	void RegisterDirectoryWatcher_AdditionalPlugins();
+	void UnregisterDirectoryWatcher_AdditionalPlugins();
 
 	/** Callback when the directory watcher detects a file/folder change */
 	void OnDirectoryChanged(const TArray<FFileChangeData>& InFileChanges);
@@ -45,8 +55,18 @@ private:
 	/** Executes when one or more folders are being deleted in the content browser */
 	void OnDeleteFolders(const TArray<FContentBrowserItemPath>& PathsToDelete);
 
+#if WITH_EDITOR
+	/** Callback when one of the directory placeholder settings has changed */
+	void OnDirectoryPlaceholderSettingsChanged(UObject* Object, FPropertyChangedEvent& PropertyChangedEvent);
+#endif // WITH_EDITOR
+
 private:
-	FDelegateHandle EngineDirectoryWatcherHandle;
+	/** Handles to each of the registerd directory watcher delegates, used to unregister them later */
+	FDelegateHandle ProjectContentHandle;
+	FDelegateHandle ProjectPluginsHandle;
+	TMap<FString, FDelegateHandle> AdditionalPluginHandles;
+
+	FDelegateHandle OnSettingsChangedHandle;
 	FDelegateHandle OnDeleteFoldersHandle;
 };
 
@@ -63,6 +83,12 @@ void FDirectoryPlaceholderModule::StartupModule()
 		FContentBrowserModule& ContentBrowserModule = FModuleManager::Get().GetModuleChecked<FContentBrowserModule>(TEXT("ContentBrowser"));
 		OnDeleteFoldersHandle = ContentBrowserModule.GetOnDeleteFolders().AddRaw(this, &FDirectoryPlaceholderModule::OnDeleteFolders);
 
+		// Register a callback to execute when one of the directory placeholders settings has changed
+#if WITH_EDITOR
+		UDirectoryPlaceholderSettings* Settings = GetMutableDefault<UDirectoryPlaceholderSettings>();
+		OnSettingsChangedHandle = Settings->OnSettingChanged().AddRaw(this, &FDirectoryPlaceholderModule::OnDirectoryPlaceholderSettingsChanged);
+#endif // WITH_EDITOR
+
 		// Extend the content browser folder context menu with the option to cleanup directory placeholders
 		if (UToolMenu* Menu = UToolMenus::Get()->ExtendMenu("ContentBrowser.FolderContextMenu"))
 		{
@@ -94,53 +120,124 @@ void FDirectoryPlaceholderModule::ShutdownModule()
 {
 	if (GIsEditor && !IsRunningCommandlet())
 	{
-		FDirectoryWatcherModule& DirectoryWatcherModule = FModuleManager::GetModuleChecked<FDirectoryWatcherModule>("DirectoryWatcher");
-		if (IDirectoryWatcher* DirectoryWatcher = DirectoryWatcherModule.Get())
-		{
-			DirectoryWatcher->UnregisterDirectoryChangedCallback_Handle(FPaths::ProjectContentDir(), EngineDirectoryWatcherHandle);
-		}
+		UnregisterDirectoryWatchers();
 
 		FContentBrowserModule& ContentBrowserModule = FModuleManager::Get().GetModuleChecked<FContentBrowserModule>(TEXT("ContentBrowser"));
 		ContentBrowserModule.GetOnDeleteFolders().Remove(OnDeleteFoldersHandle);
+
+#if WITH_EDITOR
+		if (UObjectInitialized())
+		{
+			GetMutableDefault<UDirectoryPlaceholderSettings>()->OnSettingChanged().RemoveAll(this);
+		}
+#endif
 	}
 }
 
-void FDirectoryPlaceholderModule::RegisterDirectoryWatchers()
+void FDirectoryPlaceholderModule::RegisterDirectoryWatcher(const FString& Path, FDelegateHandle& Handle)
 {
-	// Register a callback with the directory watcher to be notified about file/folder changes
+	// Register a callback with the directory watcher module to be notified about file/folder changes matching the input path
 	FDirectoryWatcherModule& DirectoryWatcherModule = FModuleManager::GetModuleChecked<FDirectoryWatcherModule>("DirectoryWatcher");
 	if (IDirectoryWatcher* DirectoryWatcher = DirectoryWatcherModule.Get())
 	{
-		if (FPaths::DirectoryExists(FPaths::ProjectContentDir()))
+		if (FPaths::DirectoryExists(Path) && !Handle.IsValid())
 		{
 			DirectoryWatcher->RegisterDirectoryChangedCallback_Handle(
-				FPaths::ProjectContentDir(),
+				Path,
 				IDirectoryWatcher::FDirectoryChanged::CreateRaw(this, &FDirectoryPlaceholderModule::OnDirectoryChanged),
-				EngineDirectoryWatcherHandle,
+				Handle,
 				IDirectoryWatcher::WatchOptions::IncludeDirectoryChanges
 			);
 		}
+	}
+}
+
+void FDirectoryPlaceholderModule::UnregisterDirectoryWatcher(const FString& Path, FDelegateHandle& Handle)
+{
+	FDirectoryWatcherModule& DirectoryWatcherModule = FModuleManager::GetModuleChecked<FDirectoryWatcherModule>("DirectoryWatcher");
+	if (IDirectoryWatcher* DirectoryWatcher = DirectoryWatcherModule.Get())
+	{
+		if (Handle.IsValid())
+		{
+			DirectoryWatcher->UnregisterDirectoryChangedCallback_Handle(Path, Handle);
+			Handle.Reset();
+		}
+	}
+}
 
-		if (FPaths::DirectoryExists(FPaths::ProjectPluginsDir()))
+void FDirectoryPlaceholderModule::RegisterDirectoryWatchers()
+{
+	const UDirectoryPlaceholderSettings* Settings = GetDefault<UDirectoryPlaceholderSettings>();
+	if (Settings->bAutomaticallyCreatePlaceholders)
+	{
+		if (Settings->bAutomaticallyCreatePlaceholdersInProjectContent)
 		{
-			DirectoryWatcher->RegisterDirectoryChangedCallback_Handle(
-				FPaths::ProjectPluginsDir(),
-				IDirectoryWatcher::FDirectoryChanged::CreateRaw(this, &FDirectoryPlaceholderModule::OnDirectoryChanged),
-				EngineDirectoryWatcherHandle,
-				IDirectoryWatcher::WatchOptions::IncludeDirectoryChanges
-			);
+			RegisterDirectoryWatcher(FPaths::ProjectContentDir(), ProjectContentHandle);
+		}
+		if (Settings->bAutomaticallyCreatePlaceholdersInProjectPlugins)
+		{
+			RegisterDirectoryWatcher(FPaths::ProjectPluginsDir(), ProjectPluginsHandle);
+		}
+		if (Settings->bAutomaticallyCreatePlaceholdersInAdditionalPlugins)
+		{
+			RegisterDirectoryWatcher_AdditionalPlugins();
 		}
 	}
 }
 
-void FDirectoryPlaceholderModule::OnDirectoryChanged(const TArray<FFileChangeData>& InFileChanges)
+void FDirectoryPlaceholderModule::UnregisterDirectoryWatchers()
 {
-	const UDirectoryPlaceholderSettings* const Settings = GetDefault<UDirectoryPlaceholderSettings>();
-	if (!Settings->bAutomaticallyCreatePlaceholders)
+	UnregisterDirectoryWatcher(FPaths::ProjectContentDir(), ProjectContentHandle);
+	UnregisterDirectoryWatcher(FPaths::ProjectPluginsDir(), ProjectPluginsHandle);
+	UnregisterDirectoryWatcher_AdditionalPlugins();
+}
+
+void FDirectoryPlaceholderModule::RegisterDirectoryWatcher_AdditionalPlugins()
+{
+	for (const FString& AdditionalPluginDirectory : IProjectManager::Get().GetAdditionalPluginDirectories())
+	{
+		FDelegateHandle& AdditionalPluginHandle = AdditionalPluginHandles.Add(AdditionalPluginDirectory);
+		RegisterDirectoryWatcher(AdditionalPluginDirectory, AdditionalPluginHandle);
+	}
+}
+
+void FDirectoryPlaceholderModule::UnregisterDirectoryWatcher_AdditionalPlugins()
+{
+	for (TPair<FString, FDelegateHandle>& AdditionalPluginHandle : AdditionalPluginHandles)
+	{
+		UnregisterDirectoryWatcher(AdditionalPluginHandle.Key, AdditionalPluginHandle.Value);
+	}
+	AdditionalPluginHandles.Empty();
+}
+
+#if WITH_EDITOR
+void FDirectoryPlaceholderModule::OnDirectoryPlaceholderSettingsChanged(UObject* Object, FPropertyChangedEvent& PropertyChangedEvent)
+{
+	const FName PropertyName = PropertyChangedEvent.GetPropertyName();
+
+	const UDirectoryPlaceholderSettings* Settings = GetDefault<UDirectoryPlaceholderSettings>();
+
+	if (PropertyName == GET_MEMBER_NAME_CHECKED(UDirectoryPlaceholderSettings, bAutomaticallyCreatePlaceholders))
+	{
+		Settings->bAutomaticallyCreatePlaceholders ? RegisterDirectoryWatchers() : UnregisterDirectoryWatchers();
+	}
+	else if (PropertyName == GET_MEMBER_NAME_CHECKED(UDirectoryPlaceholderSettings, bAutomaticallyCreatePlaceholdersInProjectContent))
 	{
-		return;
+		Settings->bAutomaticallyCreatePlaceholdersInProjectContent ? RegisterDirectoryWatcher(FPaths::ProjectContentDir(), ProjectContentHandle) : UnregisterDirectoryWatcher(FPaths::ProjectContentDir(), ProjectContentHandle);
 	}
+	else if (PropertyName == GET_MEMBER_NAME_CHECKED(UDirectoryPlaceholderSettings, bAutomaticallyCreatePlaceholdersInProjectPlugins))
+	{
+		Settings->bAutomaticallyCreatePlaceholdersInProjectPlugins ? RegisterDirectoryWatcher(FPaths::ProjectPluginsDir(), ProjectPluginsHandle) : UnregisterDirectoryWatcher(FPaths::ProjectPluginsDir(), ProjectPluginsHandle);
+	}
+	else if (PropertyName == GET_MEMBER_NAME_CHECKED(UDirectoryPlaceholderSettings, bAutomaticallyCreatePlaceholdersInAdditionalPlugins))
+	{
+		Settings->bAutomaticallyCreatePlaceholdersInAdditionalPlugins ? RegisterDirectoryWatcher_AdditionalPlugins() : UnregisterDirectoryWatcher_AdditionalPlugins();
+	}
+}
+#endif //WITH_EDITOR
 
+void FDirectoryPlaceholderModule::OnDirectoryChanged(const TArray<FFileChangeData>& InFileChanges)
+{
 	TArray<UPackage*> PackagesToSave;
 
 	for (const FFileChangeData& FileChange : InFileChanges)
@@ -161,6 +258,7 @@ void FDirectoryPlaceholderModule::OnDirectoryChanged(const TArray<FFileChangeDat
 			}
 
 			// Do not create directory placeholders in any folders excluded by the project settings
+			const UDirectoryPlaceholderSettings* const Settings = GetDefault<UDirectoryPlaceholderSettings>();
 			const bool bExclude = Algo::AnyOf(Settings->ExcludePaths, [&PackagePath](const FString& ExcludePath)
 				{
 					return FPaths::IsUnderDirectory(PackagePath, ExcludePath);
@@ -208,54 +306,34 @@ void FDirectoryPlaceholderModule::OnDirectoryChanged(const TArray<FFileChangeDat
 
 void FDirectoryPlaceholderModule::OnDeleteFolders(const TArray<FContentBrowserItemPath>& PathsToDelete)
 {
-	const UDirectoryPlaceholderSettings* const Settings = GetDefault<UDirectoryPlaceholderSettings>();
-	if (!Settings->bAutomaticallyCreatePlaceholders)
-	{
-		return;
-	}
-
 	FAssetRegistryModule& AssetRegistryModule = FModuleManager::Get().GetModuleChecked<FAssetRegistryModule>(TEXT("AssetRegistry"));
 
-	// The only place we automatically create directory placeholders is under project content ("/Game").
-	// These are also the only paths that we want to automatically delete placeholders from.
+	// Add the folders being deleted to an asset registry filter to check for existing placeholder assets
 	FARFilter Filter;
 	Filter.bRecursivePaths = true;
 	for (const FContentBrowserItemPath& Path : PathsToDelete)
 	{
-		FString InternalPath = Path.GetInternalPathName().ToString();
-		if (InternalPath.Equals(TEXT("/Game")) || InternalPath.StartsWith(TEXT("/Game/")))
-		{
-			Filter.PackagePaths.Add(Path.GetInternalPathName());
-		}
-	}
-
-	if (Filter.PackagePaths.IsEmpty())
-	{
-		return;
+		Filter.PackagePaths.Add(Path.GetInternalPathName());
 	}
 
 	// Find all of the assets in the folders being deleted
 	TArray<FAssetData> AssetDataList;
 	AssetRegistryModule.Get().GetAssets(Filter, AssetDataList);
 
-	// Check if all of the assets being deleted are placeholders, or if there are any others present
-	bool bAllPlaceholders = true;
+	// If any non-placeholder assets are found, we will not delete anything automatically
 	for (const FAssetData& AssetData : AssetDataList)
 	{
 		if (AssetData.GetClass() != UDirectoryPlaceholder::StaticClass())
 		{
-			bAllPlaceholders = false;
+			return;
 		}
 	}
 
 	// If all of the assets are directory placeholders, delete them from these folders
-	if (bAllPlaceholders)
-	{
-		const bool bShowConfirmation = false;
-		const int32 NumAssetsDeleted = ObjectTools::DeleteAssets(AssetDataList, bShowConfirmation);
+	const bool bShowConfirmation = false;
+	const int32 NumAssetsDeleted = ObjectTools::DeleteAssets(AssetDataList, bShowConfirmation);
 
-		UE_LOG(LogDirectoryPlaceholder, Verbose, TEXT("Deleted %d Directory Placeholders"), NumAssetsDeleted);
-	}
+	UE_LOG(LogDirectoryPlaceholder, Verbose, TEXT("Deleted %d Directory Placeholders"), NumAssetsDeleted);
 }
 
 IMPLEMENT_MODULE(FDirectoryPlaceholderModule, DirectoryPlaceholder);
diff --git a/Engine/Plugins/Developer/DirectoryPlaceholder/Source/DirectoryPlaceholder/Private/DirectoryPlaceholderSettings.h b/Engine/Plugins/Developer/DirectoryPlaceholder/Source/DirectoryPlaceholder/Private/DirectoryPlaceholderSettings.h
--- a/Engine/Plugins/Developer/DirectoryPlaceholder/Source/DirectoryPlaceholder/Private/DirectoryPlaceholderSettings.h
+++ b/Engine/Plugins/Developer/DirectoryPlaceholder/Source/DirectoryPlaceholder/Private/DirectoryPlaceholderSettings.h
@@ -17,14 +17,26 @@ class UDirectoryPlaceholderSettings : public UDeveloperSettings
 public:
 	UDirectoryPlaceholderSettings() = default;
 
-	/** If enabled, Directory Placeholder assets will be automatically created in new folders, and automatically deleted from folders being deleted */
+	/** If enabled, Directory Placeholder assets will be automatically created in new folders */
 	UPROPERTY(Config, EditAnywhere, Category = "Default")
 	bool bAutomaticallyCreatePlaceholders = true;
 
+	/** If enabled, Directory Placeholder assets will be automatically created in new folders that are created within the Project Content */
+	UPROPERTY(Config, EditAnywhere, Category = "Default", meta=(EditCondition = bAutomaticallyCreatePlaceholders))
+	bool bAutomaticallyCreatePlaceholdersInProjectContent = true;
+
+	/** If enabled, Directory Placeholder assets will be automatically created in new folders that are created within the Project Plugins */
+	UPROPERTY(Config, EditAnywhere, Category = "Default", meta = (EditCondition = bAutomaticallyCreatePlaceholders))
+	bool bAutomaticallyCreatePlaceholdersInProjectPlugins = true;
+
+	/** If enabled, Directory Placeholder assets will be automatically created in new folders that are created within all Additional Plugins */
+	UPROPERTY(Config, EditAnywhere, Category = "Default", meta = (EditCondition = bAutomaticallyCreatePlaceholders))
+	bool bAutomaticallyCreatePlaceholdersInAdditionalPlugins = true;
+
 	/** 
 	 * Directory Placeholder assets will not be automatically created in new folders under any of the paths in this list.
 	 * Use content browser paths, such as "/Game/MyFolder/" or "/MyPlugin/MyFolder/".
 	 */
-	UPROPERTY(Config, EditAnywhere, Category = "Default")
+	UPROPERTY(Config, EditAnywhere, Category = "Default", meta = (EditCondition = bAutomaticallyCreatePlaceholders))
 	TArray<FString> ExcludePaths;
 };
